Fix warnings generating project files.

#rnx

diff --git a/Engine/Source/Programs/PixelStreaming/WebRTCProxy/WebRTCProxy.Target.cs b/Engine/Source/Programs/PixelStreaming/WebRTCProxy/WebRTCProxy.Target.cs
--- a/Engine/Source/Programs/PixelStreaming/WebRTCProxy/WebRTCProxy.Target.cs
+++ b/Engine/Source/Programs/PixelStreaming/WebRTCProxy/WebRTCProxy.Target.cs
@@ -16,7 +16,6 @@ public class WebRTCProxyTarget : TargetRules
 		bBuildDeveloperTools = false;
 		
         // No editor needed
-        bBuildEditor = false;
 		bBuildWithEditorOnlyData = false;
 
 		// Currently this app is not linking against the engine, so we'll compile out references from Core to the rest of the engine
diff --git a/Engine/Source/Programs/UnrealMultiUserServer/UnrealMultiUserServer.Target.cs b/Engine/Source/Programs/UnrealMultiUserServer/UnrealMultiUserServer.Target.cs
--- a/Engine/Source/Programs/UnrealMultiUserServer/UnrealMultiUserServer.Target.cs
+++ b/Engine/Source/Programs/UnrealMultiUserServer/UnrealMultiUserServer.Target.cs
@@ -17,7 +17,6 @@ public class UnrealMultiUserServerTarget : TargetRules
 
         // Lean and mean also override the build developer switch, so just set all the switch it sets manually since we want developer tools (i.e. concert plugins)
         //bCompileLeanAndMeanUE = true;
-        bBuildEditor = false;
         bCompileSimplygon = false;
         bCompileSimplygonSSF = false;
         bCompileSpeedTree = false;
