Properly flush logs when shutting down the multi-user server

diff --git a/Engine/Source/Programs/UnrealMultiUserServer/Private/UnrealMultiUserServer.cpp b/Engine/Source/Programs/UnrealMultiUserServer/Private/UnrealMultiUserServer.cpp
--- a/Engine/Source/Programs/UnrealMultiUserServer/Private/UnrealMultiUserServer.cpp
+++ b/Engine/Source/Programs/UnrealMultiUserServer/Private/UnrealMultiUserServer.cpp
@@ -20,5 +20,7 @@ INT32_MAIN_INT32_ARGC_TCHAR_ARGV()
 	ServerLoopInitArgs.ServiceRole = TEXT("MultiUser");
 	ServerLoopInitArgs.ServiceFriendlyName = TEXT("Multi-User Editing Server");
 
-	return ConcertSyncServerLoop(ArgC, ArgV, ServerLoopInitArgs);
+	int32 Result = ConcertSyncServerLoop(ArgC, ArgV, ServerLoopInitArgs);
+	GLog->Flush();
+	return Result;
 }
