Fix CIS issues with merge from Main

diff --git a/Engine/Source/Programs/UnrealMultiUserServer/UnrealMultiUserServer.Target.cs b/Engine/Source/Programs/UnrealMultiUserServer/UnrealMultiUserServer.Target.cs
--- a/Engine/Source/Programs/UnrealMultiUserServer/UnrealMultiUserServer.Target.cs
+++ b/Engine/Source/Programs/UnrealMultiUserServer/UnrealMultiUserServer.Target.cs
@@ -18,8 +18,6 @@ public class UnrealMultiUserServerTarget : TargetRules
         // Lean and mean also override the build developer switch, so just set all the switch it sets manually since we want developer tools (i.e. concert plugins)
         //bCompileLeanAndMeanUE = true;
         bBuildEditor = false;
-        bCompileSimplygon = false;
-        bCompileSimplygonSSF = false;
         bCompileSpeedTree = false;
 
         // Editor-only data, however, is needed
