On Mac, added a warning log when the console output device terminate the application because this was the last window closed (with the X)
On Mac, added a warning log when UnrealMultiUserServer terminates because user closed it with CMD-Q (or from menu) or by logging out, shutting down or rebooting the computer.
On Mac, added a reason (logout, sys. restart, sys shutdown, user quitting) to the engine quit request when UnrealMultiUserServer is asked to quit.

diff --git a/Engine/Source/Programs/UnrealMultiUserServer/Private/UnrealMultiUserServer.cpp b/Engine/Source/Programs/UnrealMultiUserServer/Private/UnrealMultiUserServer.cpp
--- a/Engine/Source/Programs/UnrealMultiUserServer/Private/UnrealMultiUserServer.cpp
+++ b/Engine/Source/Programs/UnrealMultiUserServer/Private/UnrealMultiUserServer.cpp
@@ -2,9 +2,11 @@
 
 #include "ConcertSettings.h"
 #include "ConcertSyncServerLoop.h"
-
+#include "Logging/LogMacros.h"
 #include "RequiredProgramMainCPPInclude.h"
 
+DEFINE_LOG_CATEGORY_STATIC(LogMultiUserServer, Log, All)
+
 IMPLEMENT_APPLICATION(UnrealMultiUserServer, "UnrealMultiUserServer");
 
 int32 RunUnrealMultiUserServer(int ArgC, TCHAR* ArgV[])
@@ -69,14 +71,13 @@ public:
 	TCHAR** ArgV;
 };
 
+#import <Foundation/NSAppleEventDescriptor.h>
+#import <Carbon/Carbon.h>
 #include "Mac/CocoaThread.h"
 
 static CommandLineArguments GSavedCommandLine;
 
-@interface UE4AppDelegate : NSObject <NSApplicationDelegate, NSFileManagerDelegate>
-{
-}
-
+@interface UE4AppDelegate : NSObject <NSApplicationDelegate>
 @end
 
 @implementation UE4AppDelegate
@@ -87,21 +88,46 @@ static CommandLineArguments GSavedCommandLine;
 	[NSApp terminate:self];
 }
 
-- (void) runGameThread:(id)Arg
+- (void)runGameThread:(id)Arg
 {
 	FPlatformMisc::SetGracefulTerminationHandler();
 	FPlatformMisc::SetCrashHandler(nullptr);
-
 	RunUnrealMultiUserServer(GSavedCommandLine.ArgC, GSavedCommandLine.ArgV);
+}
 
-	[NSApp terminate: self];
+- (BOOL)applicationShouldTerminateAfterLastWindowClosed:(NSApplication *)Sender
+{
+	return true;
 }
 
-- (NSApplicationTerminateReply)applicationShouldTerminate:(NSApplication *)Sender;
+- (NSApplicationTerminateReply)applicationShouldTerminate:(NSApplication *)Sender
 {
 	if(!IsEngineExitRequested() || ([NSThread gameThread] && [NSThread gameThread] != [NSThread mainThread]))
 	{
-		RequestEngineExit(TEXT("UnrealMultiUserServer Requesting Exist"));
+		FString Reason;
+		NSAppleEventDescriptor* AppleEventDesc = [[NSAppleEventManager sharedAppleEventManager] currentAppleEvent];
+		NSAppleEventDescriptor* WhyDesc = [AppleEventDesc attributeDescriptorForKeyword:kEventParamReason];
+		OSType Why = [WhyDesc typeCodeValue];
+		if (Why == kAEShutDown)
+		{
+			Reason = TEXT("System Shutting Down");
+		}
+		else if (Why == kAERestart)
+		{
+			Reason = TEXT("System Restarting");
+		}
+		else if (Why == kAEReallyLogOut)
+		{
+			Reason = TEXT("User Logging Out");
+		}
+		else
+		{
+			Reason = TEXT("User Quitting (CMD-Q/Quit Menu)");
+		}
+
+		UE_LOG(LogMultiUserServer, Warning, TEXT("*** INTERRUPTED *** : %s"), *Reason);
+		RequestEngineExit(*FString::Printf(TEXT("UnrealMultiUserServer Requesting Exit: %s"), *Reason));
+		
 		return NSTerminateLater;
 	}
 	else
@@ -131,6 +157,8 @@ static CommandLineArguments GSavedCommandLine;
 	RunGameThread(self, @selector(runGameThread:));
 }
 
+@end
+
 int main(int argc, char *argv[])
 {
 	// Record the command line.
@@ -144,7 +172,6 @@ int main(int argc, char *argv[])
 	return 0;
 }
 
-@end
 
 #else // Windows/Linux
 
diff --git a/Engine/Source/Runtime/ApplicationCore/Private/Mac/MacConsoleOutputDevice.cpp b/Engine/Source/Runtime/ApplicationCore/Private/Mac/MacConsoleOutputDevice.cpp
--- a/Engine/Source/Runtime/ApplicationCore/Private/Mac/MacConsoleOutputDevice.cpp
+++ b/Engine/Source/Runtime/ApplicationCore/Private/Mac/MacConsoleOutputDevice.cpp
@@ -7,9 +7,11 @@
 #include "Mac/MacApplication.h"
 #include "Mac/MacPlatformApplicationMisc.h"
 #include "Misc/OutputDeviceHelper.h"
+#include "Misc/OutputDeviceRedirector.h"
 #include "Misc/ScopeLock.h"
 #include "Mac/CocoaThread.h"
 #include "HAL/PlatformApplicationMisc.h"
+#include "CoreGlobals.h"
 
 FMacConsoleOutputDevice::FMacConsoleOutputDevice()
 	: ConsoleHandle(NULL)
@@ -285,6 +287,11 @@ void FMacConsoleOutputDevice::SetDefaultTextColor()
 {
 	if (!MacApplication && [[NSApp orderedWindows] count] == 1)
 	{
+		UE_LOG(LogCore, Warning, TEXT("*** INTERRUPTED *** : Console Window Closed, Shutting Down"));
+		if (GLog)
+		{
+			GLog->PanicFlushThreadedLogs();
+		}
 		_Exit(0);
 	}
 }
