Add missing logging to MacOS shutdown.

diff --git a/Engine/Source/Programs/UnrealMultiUserServer/Private/UnrealMultiUserServerMain.cpp b/Engine/Source/Programs/UnrealMultiUserServer/Private/UnrealMultiUserServerMain.cpp
--- a/Engine/Source/Programs/UnrealMultiUserServer/Private/UnrealMultiUserServerMain.cpp
+++ b/Engine/Source/Programs/UnrealMultiUserServer/Private/UnrealMultiUserServerMain.cpp
@@ -96,6 +96,7 @@ static CommandLineArguments GSavedCommandLine;
 			Reason = TEXT("User Quitting (CMD-Q/Quit Menu)");
 		}
 
+		UE_LOG(LogMultiUserServer, Warning, TEXT("*** INTERRUPTED *** : SHUTTING DOWN"));
 		UE_LOG(LogMultiUserServer, Warning, TEXT("*** INTERRUPTED *** : %s"), *Reason);
 		RequestEngineExit(*FString::Printf(TEXT("UnrealMultiUserServer Requesting Exit: %s"), *Reason));
 		
