Final UX improvements for Multi-User.
Non-code changes:
- Add taskbar icons for UnrealMultiUserServer (Windows, Mac, Linux)
- Update taskbar icons for UnrealMultiUserSlateServer (Windows, Mac, Linux)
- Adjust icon on the top-left corner of window UnrealMultiUserSlateServer
- Adjust the "Create Session" button icon (shown in editor)
- Adjust the "Close server" button icon (shown in editor)
Code change:
- Disable the "Create Session" button in UnrealMultiUserSlateServer
- Make the "Launch Server" button (shown in editor) launch the slate server by default instead of the console one

diff --git a/Engine/Plugins/Developer/Concert/ConcertApp/MultiUserClient/Source/MultiUserClient/Private/Widgets/SConcertBrowser.cpp b/Engine/Plugins/Developer/Concert/ConcertApp/MultiUserClient/Source/MultiUserClient/Private/Widgets/SConcertBrowser.cpp
--- a/Engine/Plugins/Developer/Concert/ConcertApp/MultiUserClient/Source/MultiUserClient/Private/Widgets/SConcertBrowser.cpp
+++ b/Engine/Plugins/Developer/Concert/ConcertApp/MultiUserClient/Source/MultiUserClient/Private/Widgets/SConcertBrowser.cpp
@@ -148,7 +148,7 @@ public:
 	virtual void DeleteSessions(const FGuid& ServerAdminEndpointId, const TArray<FGuid>& SessionIds) override;
 	virtual bool CanDeleteActiveSession(const FGuid& ServerAdminEndpointId, const FGuid& SessionId) const override;
 	virtual bool CanDeleteArchivedSession(const FGuid& ServerAdminEndpointId, const FGuid& SessionId) const override;
-	virtual bool CanCreateSessions() const override { return true; }
+	virtual bool CanEverCreateSessions() const override { return true; }
 	//~ End IConcertSessionBrowserController Interface
 
 private:
diff --git a/Engine/Plugins/Developer/Concert/ConcertApp/MultiUserServer/Source/MultiUserServer/Private/Widgets/Browser/ConcertServerSessionBrowserController.h b/Engine/Plugins/Developer/Concert/ConcertApp/MultiUserServer/Source/MultiUserServer/Private/Widgets/Browser/ConcertServerSessionBrowserController.h
--- a/Engine/Plugins/Developer/Concert/ConcertApp/MultiUserServer/Source/MultiUserServer/Private/Widgets/Browser/ConcertServerSessionBrowserController.h
+++ b/Engine/Plugins/Developer/Concert/ConcertApp/MultiUserServer/Source/MultiUserServer/Private/Widgets/Browser/ConcertServerSessionBrowserController.h
@@ -55,7 +55,7 @@ namespace UE::MultiUserServer
 		virtual bool CanRenameArchivedSession(const FGuid& ServerAdminEndpointId, const FGuid& SessionId) const override { return true; }
 		virtual bool CanDeleteActiveSession(const FGuid& ServerAdminEndpointId, const FGuid& SessionId) const override { return true; }
 		virtual bool CanDeleteArchivedSession(const FGuid& ServerAdminEndpointId, const FGuid& SessionId) const override { return true; }
-		virtual bool CanCreateSessions() const override { return false; }
+		virtual bool CanEverCreateSessions() const override { return false; }
 		//~ Begin IConcertSessionBrowserController Interface
 
 	private:
diff --git a/Engine/Plugins/Developer/Concert/ConcertMain/Source/Concert/Public/ConcertSettings.h b/Engine/Plugins/Developer/Concert/ConcertMain/Source/Concert/Public/ConcertSettings.h
--- a/Engine/Plugins/Developer/Concert/ConcertMain/Source/Concert/Public/ConcertSettings.h
+++ b/Engine/Plugins/Developer/Concert/ConcertMain/Source/Concert/Public/ConcertSettings.h
@@ -353,7 +353,7 @@ public:
 	 *		Slate	-> UnrealMultiUserSlateServer.exe
 	 */
 	UPROPERTY(config, EditAnywhere, Category="Client Settings")
-	EConcertServerType ServerType = EConcertServerType::Console;
+	EConcertServerType ServerType = EConcertServerType::Slate;
 
 	/** 
 	 * Default server url (just a name for now) to look for on auto or default connect. 
diff --git a/Engine/Plugins/Developer/Concert/ConcertUI/ConcertSharedSlate/Source/ConcertSharedSlate/Private/ConcertFrontendStyle.cpp b/Engine/Plugins/Developer/Concert/ConcertUI/ConcertSharedSlate/Source/ConcertSharedSlate/Private/ConcertFrontendStyle.cpp
--- a/Engine/Plugins/Developer/Concert/ConcertUI/ConcertSharedSlate/Source/ConcertSharedSlate/Private/ConcertFrontendStyle.cpp
+++ b/Engine/Plugins/Developer/Concert/ConcertUI/ConcertSharedSlate/Source/ConcertSharedSlate/Private/ConcertFrontendStyle.cpp
@@ -81,13 +81,13 @@ void FConcertFrontendStyle::Initialize()
 	// Multi-User Browser 
 	StyleSet->Set("Concert.ArchiveSession",			new IMAGE_PLUGIN_BRUSH("Icons/icon_ArchiveSession_48x",			Icon24x24, IconColorAndOpacity));
 	StyleSet->Set("Concert.CancelAutoJoin",			new IMAGE_PLUGIN_BRUSH("Icons/icon_CancelAutoJoin_48x",			Icon24x24, IconColorAndOpacity));
-	StyleSet->Set("Concert.CloseServer",				new IMAGE_PLUGIN_BRUSH("Icons/icon_CloseServer_48x",			Icon24x24, IconColorAndOpacity));
+	StyleSet->Set("Concert.CloseServer",				new IMAGE_PLUGIN_BRUSH_SVG("Icons/icon_CloseServer_24",				Icon24x24, IconColorAndOpacity));
 	StyleSet->Set("Concert.DeleteSession",			new IMAGE_PLUGIN_BRUSH("Icons/icon_DeleteSession_48x",			Icon24x24, IconColorAndOpacity));
 	StyleSet->Set("Concert.JoinDefaultSession",		new IMAGE_PLUGIN_BRUSH("Icons/icon_JoinDefaultSession_48x",		Icon24x24, IconColorAndOpacity));
 	StyleSet->Set("Concert.JoinSession",				new IMAGE_PLUGIN_BRUSH("Icons/icon_JoinSelectedSession_48x",	Icon24x24, IconColorAndOpacity));
 	StyleSet->Set("Concert.LeaveSession",			new IMAGE_PLUGIN_BRUSH("Icons/icon_LeaveSession_48x",			Icon24x24, IconColorAndOpacity));
 	StyleSet->Set("Concert.NewServer",				new IMAGE_PLUGIN_BRUSH("Icons/icon_NewServer_48x",				Icon24x24, IconColorAndOpacity));
-	StyleSet->Set("Concert.NewSession",				new IMAGE_PLUGIN_BRUSH_SVG("Icons/icon_CreateMultiUser_16",		Icon24x24, IconColorAndOpacity));
+	StyleSet->Set("Concert.NewSession",				new IMAGE_PLUGIN_BRUSH_SVG("Icons/icon_CreateMultiUser_24",		Icon24x24, IconColorAndOpacity));
 	StyleSet->Set("Concert.PauseSession",			new IMAGE_PLUGIN_BRUSH("Icons/icon_PauseSession_48x",			Icon24x24, IconColorAndOpacity));
 	StyleSet->Set("Concert.RestoreSession",			new IMAGE_PLUGIN_BRUSH("Icons/icon_RestoreSession_48x",			Icon24x24, IconColorAndOpacity));
 	StyleSet->Set("Concert.ResumeSession",			new IMAGE_PLUGIN_BRUSH("Icons/icon_ResumeSession_48x",			Icon24x24, IconColorAndOpacity));
diff --git a/Engine/Plugins/Developer/Concert/ConcertUI/ConcertSharedSlate/Source/ConcertSharedSlate/Private/Session/Browser/SConcertSessionBrowser.cpp b/Engine/Plugins/Developer/Concert/ConcertUI/ConcertSharedSlate/Source/ConcertSharedSlate/Private/Session/Browser/SConcertSessionBrowser.cpp
--- a/Engine/Plugins/Developer/Concert/ConcertUI/ConcertSharedSlate/Source/ConcertSharedSlate/Private/Session/Browser/SConcertSessionBrowser.cpp
+++ b/Engine/Plugins/Developer/Concert/ConcertUI/ConcertSharedSlate/Source/ConcertSharedSlate/Private/Session/Browser/SConcertSessionBrowser.cpp
@@ -450,19 +450,21 @@ TSharedRef<SWidget> SConcertSessionBrowser::MakeButtonBar(const FArguments& InAr
 	InArgs._ExtendControllButtons.ExecuteIfBound(Extender.Get());
 	FToolBarBuilder RowBuilder(nullptr, FMultiBoxCustomization::None, Extender);
 	
-	RowBuilder.BeginSection(ControlButtonExtensionHooks::BeforeSeparator);
 	// New Session
-	RowBuilder.AddWidget(
-		ConcertBrowserUtils::MakeIconButton(
-			FConcertFrontendStyle::Get()->GetBrush("Concert.NewSession"),
-			LOCTEXT("NewButtonTooltip", "Create a new session"),
-			TAttribute<bool>(this, &SConcertSessionBrowser::IsNewButtonEnabledInternal),
-			FOnClicked::CreateSP(this, &SConcertSessionBrowser::OnNewButtonClicked)
-		)
-	);
-	RowBuilder.EndSection();
-	
-	RowBuilder.AddSeparator(ControlButtonExtensionHooks::Separator);
+	if (GetController()->CanEverCreateSessions())
+	{
+		RowBuilder.BeginSection(ControlButtonExtensionHooks::BeforeSeparator);
+		RowBuilder.AddWidget(
+			ConcertBrowserUtils::MakeIconButton(
+				FConcertFrontendStyle::Get()->GetBrush("Concert.NewSession"),
+				LOCTEXT("NewButtonTooltip", "Create a new session"),
+				TAttribute<bool>(this, &SConcertSessionBrowser::IsNewButtonEnabledInternal),
+				FOnClicked::CreateSP(this, &SConcertSessionBrowser::OnNewButtonClicked)
+			)
+		);
+		RowBuilder.EndSection();
+		RowBuilder.AddSeparator(ControlButtonExtensionHooks::Separator);
+	}
 	RowBuilder.BeginSection(ControlButtonExtensionHooks::AfterSeparator);
 	
 	// Restore (Share the same slot as Join)
@@ -1051,7 +1053,7 @@ void SConcertSessionBrowser::OnSessionSelectionChanged(TSharedPtr<FConcertSessio
 
 bool SConcertSessionBrowser::IsNewButtonEnabledInternal() const
 {
-	return GetController()->CanCreateSessions() && GetController()->GetServers().Num() > 0;
+	return GetController()->CanEverCreateSessions() && GetController()->GetServers().Num() > 0;
 }
 
 bool SConcertSessionBrowser::IsRestoreButtonEnabledInternal() const
diff --git a/Engine/Plugins/Developer/Concert/ConcertUI/ConcertSharedSlate/Source/ConcertSharedSlate/Public/Session/Browser/IConcertSessionBrowserController.h b/Engine/Plugins/Developer/Concert/ConcertUI/ConcertSharedSlate/Source/ConcertSharedSlate/Public/Session/Browser/IConcertSessionBrowserController.h
--- a/Engine/Plugins/Developer/Concert/ConcertUI/ConcertSharedSlate/Source/ConcertSharedSlate/Public/Session/Browser/IConcertSessionBrowserController.h
+++ b/Engine/Plugins/Developer/Concert/ConcertUI/ConcertSharedSlate/Source/ConcertSharedSlate/Public/Session/Browser/IConcertSessionBrowserController.h
@@ -64,7 +64,7 @@ public:
 	virtual void DeleteSessions(const FGuid& ServerAdminEndpointId, const TArray<FGuid>& SessionIds) = 0;
 	virtual bool CanDeleteActiveSession(const FGuid& ServerAdminEndpointId, const FGuid& SessionId) const = 0;
 	virtual bool CanDeleteArchivedSession(const FGuid& ServerAdminEndpointId, const FGuid& SessionId) const = 0;
-	virtual bool CanCreateSessions() const = 0;
+	virtual bool CanEverCreateSessions() const = 0;
 
 	virtual ~IConcertSessionBrowserController() = default;
 };
diff --git a/Engine/Source/Programs/UnrealMultiUserServer/Resources/Mac/make_icns.sh b/Engine/Source/Programs/UnrealMultiUserServer/Resources/Mac/make_icns.sh
new file mode 100644
--- /dev/null
+++ b/Engine/Source/Programs/UnrealMultiUserServer/Resources/Mac/make_icns.sh
@@ -0,0 +1,20 @@
+#!/bin/sh
+# Copyright Epic Games, Inc. All Rights Reserved.
+
+# Content of each /*.iconset/ folder:
+# File                 Size       DPI
+# -----------------------------------
+# icon_16x16.png       16x16       72
+# [REDACTED_EMAIL]    32x32      144
+# icon_32x32.png       32x32       72
+# [REDACTED_EMAIL]    64x64      144
+# icon_128x128.png     128x128     72
+# [REDACTED_EMAIL]  256x256    144
+# icon_256x256.png     256x256     72
+# [REDACTED_EMAIL]  512x512    144
+# icon_512x512.png     512x512     72
+# [REDACTED_EMAIL]  1024x1024  144
+# -----------------------------------
+# https://developer.apple.com/design/human-interface-guidelines/macos/icons-and-images/app-icon/
+
+iconutil --convert icns UnrealMultiUserServer.iconset
diff --git a/Engine/Source/Programs/UnrealMultiUserServer/Resources/Windows/UnrealMultiUserServer.rc b/Engine/Source/Programs/UnrealMultiUserServer/Resources/Windows/UnrealMultiUserServer.rc
new file mode 100644
--- /dev/null
+++ b/Engine/Source/Programs/UnrealMultiUserServer/Resources/Windows/UnrealMultiUserServer.rc
@@ -0,0 +1,88 @@
+// Copyright Epic Games, Inc. All Rights Reserved.
+
+#define APSTUDIO_READONLY_SYMBOLS
+/////////////////////////////////////////////////////////////////////////////
+//
+// Generated from the TEXTINCLUDE 2 resource.
+//
+#include <windows.h>
+#include "Runtime/Launch/Resources/Version.h"
+#include "Runtime/Launch/Resources/Windows/resource.h"
+
+#define UNREAL_MULTI_USER_SERVER_VERSION_MAYOR 1
+#define UNREAL_MULTI_USER_SERVER_VERSION_MINOR 0
+#define UNREAL_MULTI_USER_SERVER_VERSION_PATCH 0
+
+#define UNREAL_MULTI_USER_SERVER_VERSION_STRING "v1.01"
+
+#define VER_FILEVERSION         UNREAL_MULTI_USER_SERVER_VERSION_MAYOR,UNREAL_MULTI_USER_SERVER_VERSION_MINOR,UNREAL_MULTI_USER_SERVER_VERSION_PATCH,0
+#define VER_FILEVERSION_STR     UNREAL_MULTI_USER_SERVER_VERSION_STRING "\0"
+
+#define VER_PRODUCTVERSION      ENGINE_MAJOR_VERSION,ENGINE_MINOR_VERSION,ENGINE_PATCH_VERSION,0
+#define VER_PRODUCTVERSION_STR  ENGINE_VERSION_STRING "\0"
+
+#ifdef _DEBUG
+#define VER_FILEFLAGS           0x1L // VS_FF_DEBUG
+#else
+#define VER_FILEFLAGS           0x0L
+#endif
+
+#define IDI_UNREAL_MULTI_USER_SERVER            IDICON_UEGame // 123
+
+/////////////////////////////////////////////////////////////////////////////
+#undef APSTUDIO_READONLY_SYMBOLS
+
+/////////////////////////////////////////////////////////////////////////////
+// English (U.S.) resources
+
+#if !defined(AFX_RESOURCE_DLL) || defined(AFX_TARG_ENU)
+#ifdef _WIN32
+LANGUAGE LANG_ENGLISH, SUBLANG_ENGLISH_US
+#pragma code_page(1252)
+#endif //_WIN32
+
+/////////////////////////////////////////////////////////////////////////////
+//
+// Version
+//
+
+VS_VERSION_INFO VERSIONINFO
+FILEVERSION     VER_FILEVERSION
+PRODUCTVERSION  VER_PRODUCTVERSION
+FILEFLAGSMASK   0x17L
+FILEFLAGS       VER_FILEFLAGS
+FILEOS          0x4L
+FILETYPE        0x2L
+FILESUBTYPE     0x0L
+BEGIN
+    BLOCK "StringFileInfo"
+    BEGIN
+        BLOCK "040904b0"
+        BEGIN
+            VALUE "CompanyName",      EPIC_COMPANY_NAME
+            VALUE "LegalCopyright",   EPIC_COPYRIGHT_STRING
+            VALUE "ProductName",      EPIC_PRODUCT_NAME
+            VALUE "ProductVersion",   VER_PRODUCTVERSION_STR
+            VALUE "FileDescription",  "Unreal Multi User Server"
+            VALUE "FileVersion",      VER_FILEVERSION_STR
+            VALUE "InternalName",     "UnrealMultiUserServer"
+            VALUE "OriginalFilename", "UnrealMultiUserServer.exe"
+        END
+    END
+    BLOCK "VarFileInfo"
+    BEGIN
+        VALUE "Translation", 0x409, 1200
+    END
+END
+
+/////////////////////////////////////////////////////////////////////////////
+//
+// Icon
+//
+
+// Icon with lowest ID value placed first to ensure application icon
+// remains consistent on all systems.
+IDI_UNREAL_MULTI_USER_SERVER    ICON    "UnrealMultiUserServer.ico"
+
+#endif    // English (U.S.) resources
+/////////////////////////////////////////////////////////////////////////////
diff --git a/Engine/Source/Programs/UnrealMultiUserServer/Resources/Windows/make_ico.py b/Engine/Source/Programs/UnrealMultiUserServer/Resources/Windows/make_ico.py
new file mode 100644
--- /dev/null
+++ b/Engine/Source/Programs/UnrealMultiUserServer/Resources/Windows/make_ico.py
@@ -0,0 +1,15 @@
+# Copyright Epic Games, Inc. All Rights Reserved.
+
+from PIL import Image
+
+def MakeIcon(folder, output):
+    sizes = [256, 48, 32, 24, 16]
+    images = []
+    for size in sizes:
+        print(f"Loading {folder}/icon_{size}x{size}.png")
+        image = Image.open(f"{folder}/icon_{size}x{size}.png")
+        images.append(image)
+    print(f"Creating {output}")
+    images[0].save(f"{output}", append_images=images)
+
+MakeIcon("UnrealMultiUserServer", "UnrealMultiUserServer.ico")
diff --git a/Engine/Source/Programs/UnrealMultiUserSlateServer/Resources/Mac/make_icns.sh b/Engine/Source/Programs/UnrealMultiUserSlateServer/Resources/Mac/make_icns.sh
new file mode 100644
--- /dev/null
+++ b/Engine/Source/Programs/UnrealMultiUserSlateServer/Resources/Mac/make_icns.sh
@@ -0,0 +1,20 @@
+#!/bin/sh
+# Copyright Epic Games, Inc. All Rights Reserved.
+
+# Content of each /*.iconset/ folder:
+# File                 Size       DPI
+# -----------------------------------
+# icon_16x16.png       16x16       72
+# [REDACTED_EMAIL]    32x32      144
+# icon_32x32.png       32x32       72
+# [REDACTED_EMAIL]    64x64      144
+# icon_128x128.png     128x128     72
+# [REDACTED_EMAIL]  256x256    144
+# icon_256x256.png     256x256     72
+# [REDACTED_EMAIL]  512x512    144
+# icon_512x512.png     512x512     72
+# [REDACTED_EMAIL]  1024x1024  144
+# -----------------------------------
+# https://developer.apple.com/design/human-interface-guidelines/macos/icons-and-images/app-icon/
+
+iconutil --convert icns UnrealMultiUserSlateServer.iconset
diff --git a/Engine/Source/Programs/UnrealMultiUserSlateServer/Resources/Windows/UnrealMultiUserSlateServer.rc b/Engine/Source/Programs/UnrealMultiUserSlateServer/Resources/Windows/UnrealMultiUserSlateServer.rc
--- a/Engine/Source/Programs/UnrealMultiUserSlateServer/Resources/Windows/UnrealMultiUserSlateServer.rc
+++ b/Engine/Source/Programs/UnrealMultiUserSlateServer/Resources/Windows/UnrealMultiUserSlateServer.rc
@@ -8,7 +8,6 @@
 #include <windows.h>
 #include "Runtime/Launch/Resources/Version.h"
 #include "Runtime/Launch/Resources/Windows/resource.h"
-#include "Developer/TraceInsights/Public/Insights/Version.h"
 
 #define UNREAL_MULTI_USER_SLATE_SERVER_VERSION_MAYOR 1
 #define UNREAL_MULTI_USER_SLATE_SERVER_VERSION_MINOR 0
@@ -64,7 +63,7 @@ BEGIN
             VALUE "LegalCopyright",   EPIC_COPYRIGHT_STRING
             VALUE "ProductName",      EPIC_PRODUCT_NAME
             VALUE "ProductVersion",   VER_PRODUCTVERSION_STR
-            VALUE "FileDescription",  "Unreal Insights"
+            VALUE "FileDescription",  "Unreal Multi User Slate Server"
             VALUE "FileVersion",      VER_FILEVERSION_STR
             VALUE "InternalName",     "UnrealMultiUserSlateServer"
             VALUE "OriginalFilename", "UnrealMultiUserSlateServer.exe"
@@ -83,7 +82,7 @@ END
 
 // Icon with lowest ID value placed first to ensure application icon
 // remains consistent on all systems.
-IDI_UNREAL_MULTI_USER_SLATE_SERVER    ICON    "UnrealMultiUserSlateServer.ico"
+IDI_UNREAL_MULTI_USER_SLATE_SERVER    ICON    "UnrealMultiUserSlateServer_2.ico"
 
 #endif    // English (U.S.) resources
 /////////////////////////////////////////////////////////////////////////////
diff --git a/Engine/Source/Programs/UnrealMultiUserSlateServer/Resources/Windows/make_ico.py b/Engine/Source/Programs/UnrealMultiUserSlateServer/Resources/Windows/make_ico.py
--- a/Engine/Source/Programs/UnrealMultiUserSlateServer/Resources/Windows/make_ico.py
+++ b/Engine/Source/Programs/UnrealMultiUserSlateServer/Resources/Windows/make_ico.py
@@ -12,4 +12,4 @@ def MakeIcon(folder, output):
     print(f"Creating {output}")
     images[0].save(f"{output}", append_images=images)
 
-MakeIcon("UnrealMultiUserSlateServer", "UnrealMultiUserSlateServer.ico")
+MakeIcon("UnrealMultiUserSlateServer", "UnrealMultiUserSlateServer_2.ico")
