FPlatformString and FCString: Strcpy -> Strncpy callsite changes to handle deprecation.
#rnx

diff --git a/Engine/Source/Programs/UnrealMultiUserServer/Private/UnrealMultiUserServerMain.cpp b/Engine/Source/Programs/UnrealMultiUserServer/Private/UnrealMultiUserServerMain.cpp
--- a/Engine/Source/Programs/UnrealMultiUserServer/Private/UnrealMultiUserServerMain.cpp
+++ b/Engine/Source/Programs/UnrealMultiUserServer/Private/UnrealMultiUserServerMain.cpp
@@ -25,7 +25,7 @@ public:
 		{
 			FUTF8ToTCHAR ConvertFromUtf8(InUtf8ArgV[a]);
 			ArgV[a] = new TCHAR[ConvertFromUtf8.Length() + 1];
-			FCString::Strcpy(ArgV[a], ConvertFromUtf8.Length() + 1, ConvertFromUtf8.Get());
+			FCString::Strncpy(ArgV[a], ConvertFromUtf8.Get(), ConvertFromUtf8.Length() + 1);
 		}
 	}
 
