Fix startup of the Multi-user console server. When we added the file sharing service,
we must use the code path that delays the check for the file sharing service
because of module load order.

#Synced-CL 49551028

diff --git a/Engine/Source/Programs/UnrealMultiUserServer/Private/UnrealMultiUserServerRun.cpp b/Engine/Source/Programs/UnrealMultiUserServer/Private/UnrealMultiUserServerRun.cpp
--- a/Engine/Source/Programs/UnrealMultiUserServer/Private/UnrealMultiUserServerRun.cpp
+++ b/Engine/Source/Programs/UnrealMultiUserServer/Private/UnrealMultiUserServerRun.cpp
@@ -19,8 +19,6 @@ int32 RunUnrealMultiUserServer(int ArgC, TCHAR* ArgV[])
 		return IConcertSyncServerModule::Get().CreateFileSharingService(Role);
 	};
 	
-	ServerLoopInitArgs.FileSharingService = IConcertSyncServerModule::Get().CreateFileSharingService(Role);
-	
 	ServerLoopInitArgs.GetServerConfigFunc = [Role]() -> const UConcertServerConfig*
 	{
 		UConcertServerConfig* ServerConfig = IConcertSyncServerModule::Get().ParseServerSettings(FCommandLine::Get());
