[UnrealMultiUserServer]
* Fixed mac compile error

#Synced-CL 50306943

diff --git a/Engine/Source/Programs/UnrealMultiUserServer/Private/UnrealMultiUserServerMain.cpp b/Engine/Source/Programs/UnrealMultiUserServer/Private/UnrealMultiUserServerMain.cpp
--- a/Engine/Source/Programs/UnrealMultiUserServer/Private/UnrealMultiUserServerMain.cpp
+++ b/Engine/Source/Programs/UnrealMultiUserServer/Private/UnrealMultiUserServerMain.cpp
@@ -12,6 +12,7 @@ IMPLEMENT_APPLICATION(UnrealMultiUserServer, "UnrealMultiUserServer");
 
 #if PLATFORM_MAC // On Mac, to get a properly logging console that play nice, we need to build a mac application (.app) rather than a console application.
 
+#include "Mac/MacSystemIncludes.h"
 #include "Apple/ScopeAutoreleasePool.h"
 
 class CommandLineArguments
diff --git a/Engine/Source/Programs/UnrealMultiUserServer/UnrealMultiUserServer.Build.cs b/Engine/Source/Programs/UnrealMultiUserServer/UnrealMultiUserServer.Build.cs
--- a/Engine/Source/Programs/UnrealMultiUserServer/UnrealMultiUserServer.Build.cs
+++ b/Engine/Source/Programs/UnrealMultiUserServer/UnrealMultiUserServer.Build.cs
@@ -6,6 +6,8 @@ public class UnrealMultiUserServer : ModuleRules
 {
 	public UnrealMultiUserServer(ReadOnlyTargetRules Target) : base(Target)
 	{
+		bRequiresPlatformSDK = true;
+
 		IncludeOrderVersion = EngineIncludeOrderVersion.Latest;
 		PrivateDependencyModuleNames.AddRange(
 			new string[] {
